Added Tutorials as a context for User Feedback

TTP# 347540 - Feature Request: EDITOR: Add tutorials as an option for Quick Feedback

diff --git a/Engine/Source/Editor/UserFeedback/Private/UserFeedback.cpp b/Engine/Source/Editor/UserFeedback/Private/UserFeedback.cpp
--- a/Engine/Source/Editor/UserFeedback/Private/UserFeedback.cpp
+++ b/Engine/Source/Editor/UserFeedback/Private/UserFeedback.cpp
@@ -213,6 +213,7 @@ private:
 		ContextDisplayNames.Add(LOCTEXT("PlayInEditor", "Play In Editor"));
 		ContextDisplayNames.Add(LOCTEXT("AssetCreation", "Asset Creation"));
 		ContextDisplayNames.Add(LOCTEXT("DetailsPanel", "Details Panel"));
+		ContextDisplayNames.Add(LOCTEXT("Tutorials", "Tutorials"));
 		ContextDisplayNames.Add(LOCTEXT("Other", "Other"));
 
 		ContextMarkers.FirstEditor = ContextDisplayNames.Num();
