Slate/Core - Fixing a few cases where we were comparing TAttributes directly to literals.

diff --git a/Engine/Source/Editor/MeshPaint/Private/SMeshPaint.cpp b/Engine/Source/Editor/MeshPaint/Private/SMeshPaint.cpp
--- a/Engine/Source/Editor/MeshPaint/Private/SMeshPaint.cpp
+++ b/Engine/Source/Editor/MeshPaint/Private/SMeshPaint.cpp
@@ -1037,7 +1037,7 @@ private:
 
 	ECheckBoxState IsRadioChecked( int32 ButtonId ) const
 	{
-		return (CurrentChoice == ButtonId)
+		return (CurrentChoice.Get() == ButtonId)
 			? ECheckBoxState::Checked
 			: ECheckBoxState::Unchecked;
 	}
diff --git a/Engine/Source/Editor/UnrealEd/Private/SCurveEditor.cpp b/Engine/Source/Editor/UnrealEd/Private/SCurveEditor.cpp
--- a/Engine/Source/Editor/UnrealEd/Private/SCurveEditor.cpp
+++ b/Engine/Source/Editor/UnrealEd/Private/SCurveEditor.cpp
@@ -2516,8 +2516,11 @@ FVector2D SCurveEditor::SnapLocation(FVector2D InLocation)
 {
 	if (bSnappingEnabled.Get())
 	{
-		InLocation.X = InputSnap != 0 ? FMath::RoundToInt(InLocation.X / InputSnap.Get()) * InputSnap.Get() : InLocation.X;
-		InLocation.Y = OutputSnap != 0 ? FMath::RoundToInt(InLocation.Y / OutputSnap.Get()) * OutputSnap.Get() : InLocation.Y;
+		const float InputSnapNow = InputSnap.Get();
+		const float OutputSnapNow = OutputSnap.Get();
+
+		InLocation.X = InputSnapNow != 0 ? FMath::RoundToInt(InLocation.X / InputSnapNow) * InputSnapNow : InLocation.X;
+		InLocation.Y = OutputSnapNow != 0 ? FMath::RoundToInt(InLocation.Y / OutputSnapNow) * OutputSnapNow : InLocation.Y;
 	}
 	return InLocation;
 }
diff --git a/Engine/Source/Editor/UserFeedback/Private/UserFeedback.cpp b/Engine/Source/Editor/UserFeedback/Private/UserFeedback.cpp
--- a/Engine/Source/Editor/UserFeedback/Private/UserFeedback.cpp
+++ b/Engine/Source/Editor/UserFeedback/Private/UserFeedback.cpp
@@ -420,7 +420,7 @@ public:
 	SLATE_DEFAULT_SLOT(FArguments, Content)
 		SLATE_ATTRIBUTE(float, MinWidth)
 		SLATE_ATTRIBUTE(FText, Text)
-		SLATE_ATTRIBUTE(FName, Icon)
+		SLATE_ARGUMENT(FName, Icon)
 		SLATE_ATTRIBUTE(EFeedbackMode::Type, Mode)
 		SLATE_ATTRIBUTE(const FSlateBrush*, ArrowBrush)
 	SLATE_END_ARGS()
@@ -431,12 +431,14 @@ public:
 		MinWidth = InArgs._MinWidth;
 		
 		TSharedRef<SWidget> IconWidget = SNullWidget::NullWidget;
-		if (InArgs._Icon != NAME_None )
+		if ( InArgs._Icon != NAME_None )
 		{
-			const FSlateBrush* IconBrush = FEditorStyle::GetOptionalBrush(InArgs._Icon.Get());
+			const FSlateBrush* IconBrush = FEditorStyle::GetOptionalBrush(InArgs._Icon);
 			if (IconBrush->GetResourceName() != NAME_None)
 			{
-				IconWidget = SNew( SImage ).Image(IconBrush);
+				IconWidget = 
+					SNew( SImage )
+					.Image(IconBrush);
 			}
 		}
 
